Set rebroadcast to true by default for LiveLinkXR subjects

diff --git a/Engine/Plugins/VirtualProduction/LiveLinkXR/Source/LiveLinkXR/Private/LiveLinkXRSource.cpp b/Engine/Plugins/VirtualProduction/LiveLinkXR/Source/LiveLinkXR/Private/LiveLinkXRSource.cpp
--- a/Engine/Plugins/VirtualProduction/LiveLinkXR/Source/LiveLinkXR/Private/LiveLinkXRSource.cpp
+++ b/Engine/Plugins/VirtualProduction/LiveLinkXR/Source/LiveLinkXR/Private/LiveLinkXRSource.cpp
@@ -5,6 +5,7 @@
 #include "Async/Async.h"
 #include "Engine/Engine.h"
 #include "ILiveLinkClient.h"
+#include "LiveLinkSubjectSettings.h"
 #include "LiveLinkXR.h"
 #include "LiveLinkXRSourceSettings.h"
 #include "Misc/CoreDelegates.h"
@@ -53,6 +54,11 @@ FLiveLinkXRSource::~FLiveLinkXRSource()
 		FCoreDelegates::OnEndFrame.Remove(DeferredStartDelegateHandle);
 	}
 
+	if (Client)
+	{
+		Client->OnLiveLinkSubjectAdded().Remove(OnSubjectAddedDelegate);
+	}
+
 	Stop();
 
 	if (Thread != nullptr)
@@ -67,6 +73,8 @@ void FLiveLinkXRSource::ReceiveClient(ILiveLinkClient* InClient, FGuid InSourceG
 {
 	Client = InClient;
 	SourceGuid = InSourceGuid;
+
+	OnSubjectAddedDelegate = Client->OnLiveLinkSubjectAdded().AddRaw(this, &FLiveLinkXRSource::OnLiveLinkSubjectAdded);
 }
 
 void FLiveLinkXRSource::InitializeSettings(ULiveLinkSourceSettings* Settings)
@@ -252,6 +260,13 @@ void FLiveLinkXRSource::Send(FLiveLinkFrameDataStruct* FrameDataToSend, FName Su
 
 	if (!EncounteredSubjects.Contains(SubjectName))
 	{
+		// If the LiveLink client already knows about this subject, then it must have been added via a preset
+		// Only new subjects should be set to rebroadcast by default. Presets should respect the existing settings
+		if (!Client->GetSubjects(true, true).Contains(FLiveLinkSubjectKey(SourceGuid, SubjectName)))
+		{
+			SubjectsToRebroadcast.Add(SubjectName);
+		}
+
 		FLiveLinkStaticDataStruct StaticData(FLiveLinkTransformStaticData::StaticStruct());
 		Client->PushSubjectStaticData_AnyThread({SourceGuid, SubjectName}, ULiveLinkTransformRole::StaticClass(), MoveTemp(StaticData));
 		EncounteredSubjects.Add(SubjectName);
@@ -260,6 +275,19 @@ void FLiveLinkXRSource::Send(FLiveLinkFrameDataStruct* FrameDataToSend, FName Su
 	Client->PushSubjectFrameData_AnyThread({ SourceGuid, SubjectName }, MoveTemp(*FrameDataToSend));
 }
 
+void FLiveLinkXRSource::OnLiveLinkSubjectAdded(FLiveLinkSubjectKey InSubjectKey)
+{
+	// Set rebroadcast to true for any new subjects
+	if (SubjectsToRebroadcast.Contains(InSubjectKey.SubjectName))
+	{
+		ULiveLinkSubjectSettings* SubjectSettings = Cast<ULiveLinkSubjectSettings>(Client->GetSubjectSettings(InSubjectKey));
+		if (SubjectSettings)
+		{
+			SubjectSettings->bRebroadcastSubject = true;
+		}
+	}
+}
+
 const FString FLiveLinkXRSource::GetDeviceTypeName(EXRTrackedDeviceType DeviceType)
 {
 	switch ((int32)DeviceType)
diff --git a/Engine/Plugins/VirtualProduction/LiveLinkXR/Source/LiveLinkXR/Public/LiveLinkXRSource.h b/Engine/Plugins/VirtualProduction/LiveLinkXR/Source/LiveLinkXR/Public/LiveLinkXRSource.h
--- a/Engine/Plugins/VirtualProduction/LiveLinkXR/Source/LiveLinkXR/Public/LiveLinkXRSource.h
+++ b/Engine/Plugins/VirtualProduction/LiveLinkXR/Source/LiveLinkXR/Public/LiveLinkXRSource.h
@@ -63,6 +63,9 @@ private:
 	// Enumerate and save all connected and trackable XR devices
 	void EnumerateTrackedDevices();
 
+	// Callback when the a livelink subject has been added
+	void OnLiveLinkSubjectAdded(FLiveLinkSubjectKey InSubjectKey);
+
 private:
 	ILiveLinkClient* Client;
 
@@ -87,6 +90,9 @@ private:
 	// List of subjects we've already encountered
 	TSet<FName> EncounteredSubjects;
 
+	// List of subjects to automatically set to rebroadcast
+	TSet<FName> SubjectsToRebroadcast;
+
 	// Deferred start delegate handle.
 	FDelegateHandle DeferredStartDelegateHandle;
 
@@ -131,4 +137,7 @@ private:
 	
 	// Timestamp when we last enumerated devices
 	double LastEnumerationTimestamp = 0.0;
+
+	// Delegate for when the LiveLink client has ticked
+	FDelegateHandle OnSubjectAddedDelegate;
 };
