LiveLinkXR: Use OpenXR directly, rather than IXRTrackingSystem.

diff --git a/Engine/Plugins/VirtualProduction/LiveLinkXR/LiveLinkXR.uplugin b/Engine/Plugins/VirtualProduction/LiveLinkXR/LiveLinkXR.uplugin
--- a/Engine/Plugins/VirtualProduction/LiveLinkXR/LiveLinkXR.uplugin
+++ b/Engine/Plugins/VirtualProduction/LiveLinkXR/LiveLinkXR.uplugin
@@ -15,9 +15,13 @@
 	"IsExperimentalVersion": false,
 	"Installed": false,
 	"EnabledByDefault": false,
+	"SupportedTargetPlatforms": [
+		"Win64",
+		"Linux"
+	],
 	"Plugins": [
 		{
-			"Name": "UdpMessaging",
+			"Name": "OpenXR",
 			"Enabled": true
 		}
 	],
@@ -26,6 +30,15 @@
 			"Name": "LiveLinkXR",
 			"Type": "Runtime",
 			"LoadingPhase": "Default"
+		},
+		{
+			"Name": "LiveLinkXROpenXRExt",
+			"Type": "Runtime",
+			"LoadingPhase": "PostConfigInit",
+			"PlatformAllowList": [
+				"Win64",
+				"Linux"
+			]
 		}
 	]
 }
diff --git a/Engine/Plugins/VirtualProduction/LiveLinkXR/Source/LiveLinkXR/LiveLinkXR.Build.cs b/Engine/Plugins/VirtualProduction/LiveLinkXR/Source/LiveLinkXR/LiveLinkXR.Build.cs
--- a/Engine/Plugins/VirtualProduction/LiveLinkXR/Source/LiveLinkXR/LiveLinkXR.Build.cs
+++ b/Engine/Plugins/VirtualProduction/LiveLinkXR/Source/LiveLinkXR/LiveLinkXR.Build.cs
@@ -7,20 +7,23 @@ public class LiveLinkXR : ModuleRules
 	public LiveLinkXR(ReadOnlyTargetRules Target) : base(Target)
 	{
 		PCHUsage = ModuleRules.PCHUsageMode.UseExplicitOrSharedPCHs;
-		
+
+		PrivateIncludePaths.AddRange(
+			new string[]
+			{
+				"LiveLinkXR/Private",
+				"LiveLinkXROpenXRExt/Private",
+			}
+		);
+
 		PublicDependencyModuleNames.AddRange(
 			new string[]
 			{
 				"Core",
-				"Networking",
-				"Sockets",
 				"LiveLinkInterface",
-				"Messaging",
-				"UdpMessaging",
 			}
-			);
-			
-		
+		);
+
 		PrivateDependencyModuleNames.AddRange(
 			new string[]
 			{
@@ -29,7 +32,9 @@ public class LiveLinkXR : ModuleRules
 				"HeadMountedDisplay",
 				"Slate",
 				"SlateCore",
+				"OpenXRHMD",
+				"LiveLinkXROpenXRExt",
 			}
-			);
+		);
 	}
 }
diff --git a/Engine/Plugins/VirtualProduction/LiveLinkXR/Source/LiveLinkXR/Private/LiveLinkXRSource.cpp b/Engine/Plugins/VirtualProduction/LiveLinkXR/Source/LiveLinkXR/Private/LiveLinkXRSource.cpp
--- a/Engine/Plugins/VirtualProduction/LiveLinkXR/Source/LiveLinkXR/Private/LiveLinkXRSource.cpp
+++ b/Engine/Plugins/VirtualProduction/LiveLinkXR/Source/LiveLinkXR/Private/LiveLinkXRSource.cpp
@@ -2,25 +2,25 @@
 
 #include "LiveLinkXRSource.h"
 
-#include "Async/Async.h"
 #include "Engine/Engine.h"
+#include "HAL/RunnableThread.h"
 #include "ILiveLinkClient.h"
 #include "LiveLinkSubjectSettings.h"
 #include "LiveLinkXR.h"
+#include "LiveLinkXROpenXRExt.h"
+#include "LiveLinkXROpenXRExtModule.h"
 #include "LiveLinkXRSourceSettings.h"
 #include "Misc/CoreDelegates.h"
 #include "Misc/ScopeLock.h"
 #include "Roles/LiveLinkTransformRole.h"
+#include "Roles/LiveLinkTransformTypes.h"
 
-#define LOCTEXT_NAMESPACE "LiveLinkXRSource"
 
-static TAutoConsoleVariable<float> CVarLiveLinkXRDeviceEnumerationInterval(TEXT("LiveLinkXR.DeviceEnumerationInterval"), 1.0f, TEXT("Interval in seconds at which available SteamVR devices should be verified."));
+#define LOCTEXT_NAMESPACE "LiveLinkXRSource"
 
 
-FLiveLinkXRSource::FLiveLinkXRSource(const FLiveLinkXRConnectionSettings& Settings)
-: Client(nullptr)
-, Stopping(false)
-, Thread(nullptr)
+FLiveLinkXRSource::FLiveLinkXRSource(const FLiveLinkXRConnectionSettings& InConnectionSettings)
+	: ConnectionSettings(InConnectionSettings)
 {
 	SourceStatus = LOCTEXT("SourceStatus_NoData", "No data");
 	SourceType = LOCTEXT("SourceType_XR", "XR");
@@ -28,20 +28,28 @@ FLiveLinkXRSource::FLiveLinkXRSource(const FLiveLinkXRConnectionSettings& Settin
 
 	if (!GEngine->XRSystem.IsValid())
 	{
-		UE_LOG(LogLiveLinkXR, Error, TEXT("LiveLinkXRSource: Couldn't find a valid XR System!"));
+		UE_LOG(LogLiveLinkXR, Error, TEXT("LiveLinkXRSource: Couldn't find a valid XR System"));
+		return;
+	}
+
+	if (GEngine->XRSystem->GetSystemName() != FName(TEXT("OpenXR")))
+	{
+		UE_LOG(LogLiveLinkXR, Error, TEXT("LiveLinkXRSource: Couldn't find a compatible XR System - LiveLinkXR is only compatible with OpenXR"));
 		return;
 	}
 
-	if (GEngine->XRSystem->GetSystemName() != FName(TEXT("SteamVR")))
+	if (!FLiveLinkXROpenXRExtModule::IsAvailable())
 	{
-		UE_LOG(LogLiveLinkXR, Error, TEXT("LiveLinkXRSource: Couldn't find a compatible XR System - currently, only SteamVR is supported!"));
+		UE_LOG(LogLiveLinkXR, Error, TEXT("LiveLinkXRSource: FLiveLinkXROpenXRExtModule not available"));
 		return;
 	}
 
-	bTrackTrackers = Settings.bTrackTrackers;
-	bTrackControllers = Settings.bTrackControllers;
-	bTrackHMDs = Settings.bTrackHMDs;
-	LocalUpdateRateInHz = Settings.LocalUpdateRateInHz;
+	TSharedPtr<FLiveLinkXROpenXRExtension> XrExt = FLiveLinkXROpenXRExtModule::Get().GetExtension();
+	if (!ensure(XrExt) || !XrExt->IsSupported())
+	{
+		UE_LOG(LogLiveLinkXR, Error, TEXT("LiveLinkXRSource: Extension plugin failed to initialize"));
+		return;
+	}
 
 	DeferredStartDelegateHandle = FCoreDelegates::OnEndFrame.AddRaw(this, &FLiveLinkXRSource::Start);
 }
@@ -79,22 +87,23 @@ void FLiveLinkXRSource::ReceiveClient(ILiveLinkClient* InClient, FGuid InSourceG
 
 void FLiveLinkXRSource::InitializeSettings(ULiveLinkSourceSettings* Settings)
 {
+	SavedSourceSettings = Cast<ULiveLinkXRSourceSettings>(Settings);
+	if (!ensure(SavedSourceSettings))
+	{
+		return;
+	}
+
+	LocalUpdateRateInHz = SavedSourceSettings->LocalUpdateRateInHz;
 }
 
 void FLiveLinkXRSource::Update()
 {
-	const double CurrentTime = FPlatformTime::Seconds();
-	if((CurrentTime - LastEnumerationTimestamp) > CVarLiveLinkXRDeviceEnumerationInterval.GetValueOnGameThread())
-	{
-		LastEnumerationTimestamp = CurrentTime;
-		EnumerateTrackedDevices();
-	}	
 }
 
 bool FLiveLinkXRSource::IsSourceStillValid() const
 {
 	// Source is valid if we have a valid thread
-	const bool bIsSourceValid = !Stopping && (Thread != nullptr);
+	const bool bIsSourceValid = !bStopping && (Thread != nullptr);
 	return bIsSourceValid;
 }
 
@@ -105,76 +114,9 @@ bool FLiveLinkXRSource::RequestSourceShutdown()
 	return true;
 }
 
-void FLiveLinkXRSource::OnSettingsChanged(ULiveLinkSourceSettings* Settings, const FPropertyChangedEvent& PropertyChangedEvent)
+TSubclassOf<ULiveLinkSourceSettings> FLiveLinkXRSource::GetSettingsClass() const
 {
-	ILiveLinkSource::OnSettingsChanged(Settings, PropertyChangedEvent);
-}
-
-void FLiveLinkXRSource::EnumerateTrackedDevices()
-{
-	if (!GEngine->XRSystem.IsValid())
-	{
-		return;
-	}
-
-	if (GEngine->XRSystem->GetSystemName() != FName(TEXT("SteamVR")))
-	{
-		return;
-	}
-
-	// Create subject names for all requested tracked devices
-	TArray<int32> AllTrackedDevices;
-
-	if (!GEngine->XRSystem->EnumerateTrackedDevices(AllTrackedDevices, EXRTrackedDeviceType::Any))
-	{
-		return;
-	}
-
-	FScopeLock Lock(&TrackedDeviceCriticalSection);
-
-	for (int32 TrackerIndex = 0; TrackerIndex < AllTrackedDevices.Num(); TrackerIndex++)
-	{
-		FString SubjectStringId = GEngine->XRSystem->GetSystemName().ToString();
-		bool bValidDevice = false;
-		switch ((int32)GEngine->XRSystem->GetTrackedDeviceType(AllTrackedDevices[TrackerIndex]))
-		{
-		case (int32)EXRTrackedDeviceType::Other:
-			if (bTrackTrackers)
-			{
-				SubjectStringId += TEXT("Tracker_");
-				bValidDevice = true;
-			}
-			break;
-
-		case (int32)EXRTrackedDeviceType::HeadMountedDisplay:
-			if (bTrackHMDs)
-			{
-				SubjectStringId += TEXT("HMD_");
-				bValidDevice = true;
-			}
-			break;
-
-		case (int32)EXRTrackedDeviceType::Controller:
-			if (bTrackControllers)
-			{
-				SubjectStringId += TEXT("Controller_");
-				bValidDevice = true;
-			}
-			break;
-		}
-
-		if (bValidDevice)
-		{
-			SubjectStringId += GEngine->XRSystem->GetTrackedDevicePropertySerialNumber(AllTrackedDevices[TrackerIndex]);
-			const FName SubjectName(SubjectStringId);
-			if(TrackedDevices.ContainsByPredicate([SubjectName](const FTrackedDeviceInfo& Other){ return Other.SubjectName == SubjectName; }) == false)
-			{
-				FTrackedDeviceInfo NewDevice(AllTrackedDevices[TrackerIndex], GEngine->XRSystem->GetTrackedDeviceType(AllTrackedDevices[TrackerIndex]), SubjectName);
-				UE_LOG(LogLiveLinkXR, Log, TEXT("LiveLinkXRSource: Found a tracked device with DeviceId %d and named it %s"), NewDevice.DeviceId, *SubjectStringId);
-				TrackedDevices.Add(MoveTemp(NewDevice));
-			}
-		}
-	}
+	return ULiveLinkXRSourceSettings::StaticClass();
 }
 
 // FRunnable interface
@@ -184,68 +126,71 @@ void FLiveLinkXRSource::Start()
 
 	FCoreDelegates::OnEndFrame.Remove(DeferredStartDelegateHandle);
 	DeferredStartDelegateHandle.Reset();
-
-	EnumerateTrackedDevices();
 	
 	SourceStatus = LOCTEXT("SourceStatus_Receiving", "Receiving");
 
+	static std::atomic<int32> ReceiverIndex = 0;
 	ThreadName = "LiveLinkXR Receiver ";
-	ThreadName.AppendInt(FAsyncThreadIndex::GetNext());
+	ThreadName.AppendInt(++ReceiverIndex);
 
 	Thread = FRunnableThread::Create(this, *ThreadName, 128 * 1024, TPri_AboveNormal, FPlatformAffinity::GetPoolThreadMask());
 }
 
 void FLiveLinkXRSource::Stop()
 {
-	Stopping = true;
+	bStopping = true;
 }
 
 uint32 FLiveLinkXRSource::Run()
 {
-	const double CheckDeltaTime = 1.0f / (double)LocalUpdateRateInHz;
-	double StartTime = FApp::GetCurrentTime();
-	double EndTime = StartTime + CheckDeltaTime;
+	TWeakPtr<FLiveLinkXROpenXRExtension> WeakExt;
+	if (ensure(FLiveLinkXROpenXRExtModule::IsAvailable()))
+	{
+		WeakExt = FLiveLinkXROpenXRExtModule::Get().GetExtension();
+		ensure(WeakExt.IsValid());
+	}
 
-	while (!Stopping)
+	TMap<FName, FTransform> SubjectPoses;
+	double LastFrameTimeSec = -DBL_MAX;
+	while (!bStopping)
 	{
 		// Send new poses at the user specified update rate
-		if (EndTime - StartTime >= CheckDeltaTime)
+		const double FrameIntervalSec = 1.0 / LocalUpdateRateInHz;
+		const double TimeNowSec = FApp::GetCurrentTime();
+		if (TimeNowSec >= (LastFrameTimeSec + FrameIntervalSec))
 		{
-			StartTime = FApp::GetCurrentTime();
+			LastFrameTimeSec = TimeNowSec;
+
+			// Send new poses at the user specified update rate
+			SubjectPoses.Reset();
 
-			// Make a copy of tracked device array to avoid locking for too long
-			TArray<FTrackedDeviceInfo> TrackedDevicesCopy;
+			if (TSharedPtr<FLiveLinkXROpenXRExtension> XrExt = WeakExt.Pin())
 			{
-				FScopeLock Lock(&TrackedDeviceCriticalSection);
-				TrackedDevicesCopy = TrackedDevices;
+				XrExt->GetSubjectPoses(SubjectPoses);
 			}
-
-			// Send new poses at the user specified update rate
-			for (const FTrackedDeviceInfo& Device : TrackedDevicesCopy)
+			else
 			{
-				FQuat Orientation;
-				FVector Position;
-
-				if (GEngine->XRSystem->IsTracking(Device.DeviceId) && GEngine->XRSystem->GetCurrentPose(Device.DeviceId, Orientation, Position))
-				{
-					FLiveLinkFrameDataStruct FrameData(FLiveLinkTransformFrameData::StaticStruct());
-					FLiveLinkTransformFrameData* TransformFrameData = FrameData.Cast<FLiveLinkTransformFrameData>();
-					TransformFrameData->Transform = FTransform(Orientation, Position);
-
-					// These don't change frame to frame, so they really should be in static data. However, there is no MetaData in LiveLink static data :(
-					TransformFrameData->MetaData.StringMetaData.Add(FName(TEXT("DeviceId")), FString::Printf(TEXT("%d"), Device.DeviceId));
-					TransformFrameData->MetaData.StringMetaData.Add(FName(TEXT("DeviceType")), GetDeviceTypeName(Device.DeviceType));
-					TransformFrameData->MetaData.StringMetaData.Add(FName(TEXT("DeviceControlId")), Device.SubjectName.ToString());
-
-					Send(&FrameData, Device.SubjectName);
-				}
+				UE_LOG(LogLiveLinkXR, Error, TEXT("Source thread couldn't pin FLiveLinkXROpenXRExtension; stopping"));
+				bStopping = true;
+				return 1;
 			}
 
-			FrameCounter++;
+			for (const TPair<FName, FTransform>& SubjectPose : SubjectPoses)
+			{
+				const FName& SubjectName = SubjectPose.Key;
+
+				FLiveLinkFrameDataStruct FrameData(FLiveLinkTransformFrameData::StaticStruct());
+				FLiveLinkTransformFrameData* TransformFrameData = FrameData.Cast<FLiveLinkTransformFrameData>();
+				TransformFrameData->Transform = SubjectPose.Value;
+
+				// These don't change frame to frame, so they really should be in static data. However, there is no MetaData in LiveLink static data :(
+				TransformFrameData->MetaData.StringMetaData.Add(FName(TEXT("DeviceControlId")), SubjectName.ToString());
+
+				Send(&FrameData, SubjectName);
+			}
 		}
 
 		FPlatformProcess::Sleep(0.001f);
-		EndTime = FApp::GetCurrentTime();
 	}
 	
 	return 0;
@@ -253,7 +198,7 @@ uint32 FLiveLinkXRSource::Run()
 
 void FLiveLinkXRSource::Send(FLiveLinkFrameDataStruct* FrameDataToSend, FName SubjectName)
 {
-	if (Stopping || (Client == nullptr))
+	if (bStopping || (Client == nullptr))
 	{
 		return;
 	}
@@ -301,4 +246,5 @@ const FString FLiveLinkXRSource::GetDeviceTypeName(EXRTrackedDeviceType DeviceTy
 	return FString::Printf(TEXT("Unknown (%i)"), (int32)DeviceType);
 }
 
+
 #undef LOCTEXT_NAMESPACE
diff --git a/Engine/Plugins/VirtualProduction/LiveLinkXR/Source/LiveLinkXR/Public/LiveLinkXRSource.h b/Engine/Plugins/VirtualProduction/LiveLinkXR/Source/LiveLinkXR/Public/LiveLinkXRSource.h
--- a/Engine/Plugins/VirtualProduction/LiveLinkXR/Source/LiveLinkXR/Public/LiveLinkXRSource.h
+++ b/Engine/Plugins/VirtualProduction/LiveLinkXR/Source/LiveLinkXR/Public/LiveLinkXRSource.h
@@ -3,27 +3,24 @@
 #pragma once
 
 #include "ILiveLinkSource.h"
-#include "LiveLinkXRConnectionSettings.h"
-#include "LiveLinkXRSourceSettings.h"
-#include "Roles/LiveLinkTransformTypes.h"
-
-#include "Delegates/IDelegateInstance.h"
-#include "MessageEndpoint.h"
-#include "IMessageContext.h"
-#include "HAL/ThreadSafeBool.h"
 #include "HAL/Runnable.h"
+#include "IXRTrackingSystem.h" // for EXRTrackedDeviceType; no longer used but pending deprecation
+#include "LiveLinkXRConnectionSettings.h"
 
-#include "IXRTrackingSystem.h"		// for GEngine->XRSystem and EXRTrackedDeviceType
+#include <atomic>
 
-struct ULiveLinkXRSettings;
 
 class ILiveLinkClient;
+class ULiveLinkXRSourceSettings;
 
-class LIVELINKXR_API FLiveLinkXRSource : public ILiveLinkSource, public FRunnable, public TSharedFromThis<FLiveLinkXRSource>
+
+class LIVELINKXR_API FLiveLinkXRSource
+	: public ILiveLinkSource
+	, public FRunnable
+	, public TSharedFromThis<FLiveLinkXRSource>
 {
 public:
-
-	FLiveLinkXRSource(const FLiveLinkXRConnectionSettings& ConnectionSettings);
+	FLiveLinkXRSource(const FLiveLinkXRConnectionSettings& InConnectionSettings);
 
 	virtual ~FLiveLinkXRSource();
 
@@ -41,8 +38,7 @@ public:
 	virtual FText GetSourceMachineName() const override { return SourceMachineName; }
 	virtual FText GetSourceStatus() const override { return SourceStatus; }
 
-	virtual TSubclassOf<ULiveLinkSourceSettings> GetSettingsClass() const override { return ULiveLinkXRSourceSettings::StaticClass(); }
-	virtual void OnSettingsChanged(ULiveLinkSourceSettings* Settings, const FPropertyChangedEvent& PropertyChangedEvent) override;
+	virtual TSubclassOf<ULiveLinkSourceSettings> GetSettingsClass() const override;
 
 	// End ILiveLinkSource Interface
 
@@ -57,34 +53,31 @@ public:
 	// End FRunnable Interface
 
 	void Send(FLiveLinkFrameDataStruct* FrameDataToSend, FName SubjectName);
+
+	UE_DEPRECATED(5.1, "LiveLinkXR no longer uses EXRTrackedDeviceType.")
 	const FString GetDeviceTypeName(EXRTrackedDeviceType DeviceType);
 
 private:
-	// Enumerate and save all connected and trackable XR devices
-	void EnumerateTrackedDevices();
-
 	// Callback when the a livelink subject has been added
 	void OnLiveLinkSubjectAdded(FLiveLinkSubjectKey InSubjectKey);
 
 private:
-	ILiveLinkClient* Client;
+	ILiveLinkClient* Client = nullptr;
 
 	// Our identifier in LiveLink
 	FGuid SourceGuid;
 
-	FMessageAddress ConnectionAddress;
-
 	FText SourceType;
 	FText SourceMachineName;
 	FText SourceStatus;
 	
-	// Threadsafe Bool for terminating the main thread loop
-	FThreadSafeBool Stopping;
+	// Threadsafe flag for terminating the main thread loop
+	std::atomic<bool> bStopping = false;
 	
-	// Thread to run socket operations on
-	FRunnableThread* Thread;
+	// Thread to update poses from
+	FRunnableThread* Thread = nullptr;
 	
-	// Name of the sockets thread
+	// Name of the update thread
 	FString ThreadName;
 	
 	// List of subjects we've already encountered
@@ -99,45 +92,15 @@ private:
 	// frame counter for data
 	int32 FrameCounter = 0;
 
-	// Track all SteamVR tracker "pucks"
-	bool bTrackTrackers = true;
-
-	// Track all controllers
-	bool bTrackControllers = false;
-
-	// Track all HMDs
-	bool bTrackHMDs = false;
-
 	// Update rate (in Hz) at which to read the tracking data for each device
-	uint32 LocalUpdateRateInHz;
-
-	struct FTrackedDeviceInfo
-	{
-		FTrackedDeviceInfo(int32 InDeviceId, EXRTrackedDeviceType InType, FName InName)
-			: DeviceId(InDeviceId)
-			, DeviceType(InType)
-			, SubjectName(InName)
-		{}
-
-		// Device Id for local SteamVR
-		int32 DeviceId = INDEX_NONE;
-
-		// Device type
-		EXRTrackedDeviceType DeviceType = EXRTrackedDeviceType::Invalid;
-		
-		// Subject name of this device
-		FName SubjectName;
-	};
-
-	// Array of found devices
-	TArray<FTrackedDeviceInfo> TrackedDevices;
+	std::atomic<uint32> LocalUpdateRateInHz;
 
 	// Critical section protection tracked devices
 	FCriticalSection TrackedDeviceCriticalSection;
-	
-	// Timestamp when we last enumerated devices
-	double LastEnumerationTimestamp = 0.0;
 
 	// Delegate for when the LiveLink client has ticked
 	FDelegateHandle OnSubjectAddedDelegate;
+
+	const FLiveLinkXRConnectionSettings ConnectionSettings;
+	ULiveLinkXRSourceSettings* SavedSourceSettings = nullptr;
 };
diff --git a/Engine/Plugins/VirtualProduction/LiveLinkXR/Source/LiveLinkXR/Public/LiveLinkXRSourceSettings.h b/Engine/Plugins/VirtualProduction/LiveLinkXR/Source/LiveLinkXR/Public/LiveLinkXRSourceSettings.h
--- a/Engine/Plugins/VirtualProduction/LiveLinkXR/Source/LiveLinkXR/Public/LiveLinkXRSourceSettings.h
+++ b/Engine/Plugins/VirtualProduction/LiveLinkXR/Source/LiveLinkXR/Public/LiveLinkXRSourceSettings.h
@@ -6,11 +6,14 @@
 #include "LiveLinkSourceSettings.h"
 #include "LiveLinkXRSourceSettings.generated.h"
 
+
 UCLASS()
 class LIVELINKXR_API ULiveLinkXRSourceSettings : public ULiveLinkSourceSettings
 {
 	GENERATED_BODY()
 
 public:
-	// Placeholder for potential future source settings
+	/** Update rate (in Hz) at which to read the tracking data for each device */
+	UPROPERTY(EditAnywhere, Category="LiveLinkXR", meta=(ClampMin=1, ClampMax=1000))
+	uint32 LocalUpdateRateInHz = 60;
 };
diff --git a/Engine/Plugins/VirtualProduction/LiveLinkXR/Source/LiveLinkXROpenXRExt/LiveLinkXROpenXRExt.Build.cs b/Engine/Plugins/VirtualProduction/LiveLinkXR/Source/LiveLinkXROpenXRExt/LiveLinkXROpenXRExt.Build.cs
new file mode 100644
--- /dev/null
+++ b/Engine/Plugins/VirtualProduction/LiveLinkXR/Source/LiveLinkXROpenXRExt/LiveLinkXROpenXRExt.Build.cs
@@ -0,0 +1,37 @@
+// Copyright Epic Games, Inc. All Rights Reserved.
+
+using UnrealBuildTool;
+using System.IO;
+
+public class LiveLinkXROpenXRExt : ModuleRules
+{
+	public LiveLinkXROpenXRExt(ReadOnlyTargetRules Target) : base(Target)
+	{
+		PCHUsage = ModuleRules.PCHUsageMode.UseExplicitOrSharedPCHs;
+
+		var EngineDir = Path.GetFullPath(Target.RelativeEnginePath);
+		PrivateIncludePaths.AddRange(
+            new string[] {
+                EngineDir + "/Plugins/Runtime/OpenXR/Source/OpenXRHMD/Private",
+            }
+		);
+
+		PublicDependencyModuleNames.AddRange(
+			new string[]
+			{
+				"Core"
+			}
+		);
+
+		PrivateDependencyModuleNames.AddRange(
+			new string[]
+			{
+				"CoreUObject",
+				"Engine",
+				"HeadMountedDisplay",
+				"InputCore",
+				"OpenXRHMD",
+			}
+		);
+	}
+}
diff --git a/Engine/Plugins/VirtualProduction/LiveLinkXR/Source/LiveLinkXROpenXRExt/Private/LiveLinkXROpenXRExt.cpp b/Engine/Plugins/VirtualProduction/LiveLinkXR/Source/LiveLinkXROpenXRExt/Private/LiveLinkXROpenXRExt.cpp
new file mode 100644
--- /dev/null
+++ b/Engine/Plugins/VirtualProduction/LiveLinkXR/Source/LiveLinkXROpenXRExt/Private/LiveLinkXROpenXRExt.cpp
@@ -0,0 +1,226 @@
+// Copyright Epic Games, Inc. All Rights Reserved.
+
+#include "LiveLinkXROpenXRExt.h"
+#include "LiveLinkXROpenXRExtModule.h"
+
+#include "Engine/Engine.h"
+#include "IOpenXRHMDModule.h"
+#include "IXRTrackingSystem.h"
+#include "OpenXRCore.h"
+#include "OpenXRHMD.h"
+
+#if PLATFORM_WINDOWS
+	#include "Microsoft/MinimalWindowsApi.h"
+#endif
+
+
+#define LOCTEXT_NAMESPACE "LiveLinkXROpenXR"
+
+
+XrPath GetPath(XrInstance Instance, const char* PathString)
+{
+	XrPath Path = XR_NULL_PATH;
+	XrResult Result = xrStringToPath(Instance, PathString, &Path);
+	check(XR_SUCCEEDED(Result));
+	return Path;
+}
+
+XrPath GetPath(XrInstance Instance, const FString& PathString)
+{
+	return GetPath(Instance, StringCast<ANSICHAR>(*PathString).Get());
+}
+
+
+
+FLiveLinkXROpenXRExtension::FLiveLinkXROpenXRExtension()
+{
+	RegisterOpenXRExtensionModularFeature();
+}
+
+
+FLiveLinkXROpenXRExtension::~FLiveLinkXROpenXRExtension()
+{
+	UnregisterOpenXRExtensionModularFeature();
+}
+
+
+void FLiveLinkXROpenXRExtension::GetSubjectPoses(TMap<FName, FTransform>& OutPoseMap)
+{
+	XrTime PoseTime;
+	if (!XR_ENSURE(GetXrTimeNow(PoseTime)))
+	{
+		return;
+	}
+
+	XrSpace RefSpace = GetReferenceSpace();
+	if (RefSpace == XR_NULL_HANDLE)
+	{
+		return;
+	}
+
+	TArray<int32> Devices;
+	FOpenXRHMD* OpenXRHMD = (FOpenXRHMD*)GEngine->XRSystem.Get();
+	if (OpenXRHMD->EnumerateTrackedDevices(Devices))
+	{
+		for (int32 DeviceId : Devices)
+		{
+			FName Path;
+			FTransform Location;
+			if (LocateSpace(OpenXRHMD->GetTrackedDeviceSpace(DeviceId), RefSpace, PoseTime, Location))
+			{
+				OutPoseMap.Emplace(FOpenXRPath(OpenXRHMD->GetTrackedDevicePath(DeviceId)), Location);
+			}
+		}
+	}
+}
+
+
+bool FLiveLinkXROpenXRExtension::GetOptionalExtensions(TArray<const ANSICHAR*>& OutExtensions)
+{
+#if PLATFORM_WINDOWS
+	OutExtensions.Add("XR_KHR_win32_convert_performance_counter_time");
+#endif
+
+	return true;
+}
+
+
+void FLiveLinkXROpenXRExtension::PostCreateInstance(XrInstance InInstance)
+{
+	Instance = InInstance;
+
+	bIsSupported = false;
+
+	if (!IOpenXRHMDModule::IsAvailable())
+	{
+		UE_LOG(LogLiveLinkXROpenXRExt, Error, TEXT("IOpenXRHMDModule not available"));
+		return;
+	}
+
+#if PLATFORM_WINDOWS
+	if (!IOpenXRHMDModule::Get().IsExtensionEnabled("XR_KHR_win32_convert_performance_counter_time"))
+	{
+		UE_LOG(LogLiveLinkXROpenXRExt, Error, TEXT("XR_KHR_win32_convert_performance_counter_time extension not available"));
+		return;
+	}
+
+	if (!XR_ENSURE(xrGetInstanceProcAddr(Instance, "xrConvertWin32PerformanceCounterToTimeKHR",
+		reinterpret_cast<PFN_xrVoidFunction*>(&PfnXrConvertWin32PerformanceCounterToTimeKHR))))
+	{
+		UE_LOG(LogLiveLinkXROpenXRExt, Error, TEXT("Couldn't resolve xrConvertWin32PerformanceCounterToTimeKHR"));
+		return;
+	}
+#else
+	UE_LOG(LogLiveLinkXROpenXRExt, Error, TEXT("Support not implemented for this platform"));
+	return;
+#endif
+
+	bIsSupported = true;
+}
+
+
+void FLiveLinkXROpenXRExtension::PostCreateSession(XrSession InSession)
+{
+	// Cache WorldToMetersScale.
+	static FName SystemName(TEXT("OpenXR"));
+	if (ensure(GEngine->XRSystem.IsValid()) && ensure(GEngine->XRSystem->GetSystemName() == SystemName))
+	{
+		WorldToMetersScale = GEngine->XRSystem.Get()->GetWorldToMetersScale();
+	}
+
+	// Create reference spaces. (Borrowed from FOpenXRHMD::OnStereoStartup)
+	uint32_t ReferenceSpacesCount;
+	XR_ENSURE(xrEnumerateReferenceSpaces(InSession, 0, &ReferenceSpacesCount, nullptr));
+
+	TArray<XrReferenceSpaceType> Spaces;
+	Spaces.SetNum(ReferenceSpacesCount);
+	for (XrReferenceSpaceType& SpaceIter : Spaces)
+	{
+		// Initialize spaces array with valid enum values (avoid triggering validation error).
+		SpaceIter = XR_REFERENCE_SPACE_TYPE_VIEW;
+	}
+	XR_ENSURE(xrEnumerateReferenceSpaces(InSession, (uint32_t)Spaces.Num(), &ReferenceSpacesCount, Spaces.GetData()));
+	ensure(ReferenceSpacesCount == Spaces.Num());
+
+	XrReferenceSpaceCreateInfo SpaceInfo{ XR_TYPE_REFERENCE_SPACE_CREATE_INFO, nullptr };
+	SpaceInfo.poseInReferenceSpace = ToXrPose(FTransform::Identity);
+
+	ensure(Spaces.Contains(XR_REFERENCE_SPACE_TYPE_LOCAL));
+	SpaceInfo.referenceSpaceType = XR_REFERENCE_SPACE_TYPE_LOCAL;
+	XR_ENSURE(xrCreateReferenceSpace(InSession, &SpaceInfo, &LocalSpace));
+
+	if (Spaces.Contains(XR_REFERENCE_SPACE_TYPE_STAGE))
+	{
+		TrackingSpaceType = XR_REFERENCE_SPACE_TYPE_STAGE;
+		SpaceInfo.referenceSpaceType = TrackingSpaceType;
+		XR_ENSURE(xrCreateReferenceSpace(InSession, &SpaceInfo, &StageSpace));
+	}
+}
+
+
+void FLiveLinkXROpenXRExtension::OnDestroySession(XrSession InSession)
+{
+	if (LocalSpace != XR_NULL_HANDLE)
+	{
+		XR_ENSURE(xrDestroySpace(LocalSpace));
+		LocalSpace = XR_NULL_HANDLE;
+	}
+
+	if (StageSpace != XR_NULL_HANDLE)
+	{
+		XR_ENSURE(xrDestroySpace(StageSpace));
+		StageSpace = XR_NULL_HANDLE;
+	}
+}
+
+
+XrSpace FLiveLinkXROpenXRExtension::GetReferenceSpace() const
+{
+	switch (TrackingSpaceType)
+	{
+		case XR_REFERENCE_SPACE_TYPE_STAGE: return StageSpace;
+		case XR_REFERENCE_SPACE_TYPE_LOCAL: return LocalSpace;
+		default: checkNoEntry(); return XR_NULL_HANDLE;
+	}
+}
+
+
+XrResult FLiveLinkXROpenXRExtension::GetXrTimeNow(XrTime& OutXrTime) const
+{
+#if PLATFORM_WINDOWS
+	if (ensure(Instance) && ensure(PfnXrConvertWin32PerformanceCounterToTimeKHR))
+	{
+		Windows::LARGE_INTEGER CyclesNow;
+		Windows::QueryPerformanceCounter(&CyclesNow);
+
+		return PfnXrConvertWin32PerformanceCounterToTimeKHR(Instance, &CyclesNow, &OutXrTime);
+	}
+#endif
+
+	return XR_ERROR_FEATURE_UNSUPPORTED;
+}
+
+bool FLiveLinkXROpenXRExtension::LocateSpace(XrSpace InSpace, XrSpace InRefSpace, XrTime InTime, FTransform& OutTransform)
+{
+	XrSpaceLocation XrLocation{ XR_TYPE_SPACE_LOCATION, nullptr };
+	if (XR_SUCCEEDED(xrLocateSpace(InSpace, InRefSpace, InTime, &XrLocation)))
+	{
+		if ((XrLocation.locationFlags & XR_SPACE_LOCATION_POSITION_VALID_BIT) == 0 ||
+			(XrLocation.locationFlags & XR_SPACE_LOCATION_POSITION_TRACKED_BIT) == 0)
+		{
+			return false;
+		}
+
+		if ((XrLocation.locationFlags & XR_SPACE_LOCATION_ORIENTATION_VALID_BIT) == 0 ||
+			(XrLocation.locationFlags & XR_SPACE_LOCATION_ORIENTATION_TRACKED_BIT) == 0)
+		{
+			return false;
+		}
+
+		OutTransform = FTransform(
+			ToFQuat(XrLocation.pose.orientation), ToFVector(XrLocation.pose.position, WorldToMetersScale)
+		);
+		return true;
+	}
+	return false;
+}
diff --git a/Engine/Plugins/VirtualProduction/LiveLinkXR/Source/LiveLinkXROpenXRExt/Private/LiveLinkXROpenXRExt.h b/Engine/Plugins/VirtualProduction/LiveLinkXR/Source/LiveLinkXROpenXRExt/Private/LiveLinkXROpenXRExt.h
new file mode 100644
--- /dev/null
+++ b/Engine/Plugins/VirtualProduction/LiveLinkXR/Source/LiveLinkXROpenXRExt/Private/LiveLinkXROpenXRExt.h
@@ -0,0 +1,61 @@
+// Copyright Epic Games, Inc. All Rights Reserved.
+
+#pragma once
+
+#include "IOpenXRExtensionPlugin.h"
+#include "OpenXRCore.h"
+
+
+#if PLATFORM_WINDOWS
+namespace Windows
+{
+	union LARGE_INTEGER;
+}
+#endif
+
+
+class LIVELINKXROPENXREXT_API FLiveLinkXROpenXRExtension
+	: public IOpenXRExtensionPlugin
+	, public TSharedFromThis<FLiveLinkXROpenXRExtension>
+{
+public:
+	FLiveLinkXROpenXRExtension();
+	virtual ~FLiveLinkXROpenXRExtension();
+
+	bool IsSupported() const { return bIsSupported; }
+	void GetSubjectPoses(TMap<FName, FTransform>& OutPoseMap);
+
+	//~ Begin IOpenXRExtensionPlugin interface
+	virtual FString GetDisplayName() override
+	{
+		return FString(TEXT("LiveLinkXR"));
+	}
+
+	virtual bool GetOptionalExtensions(TArray<const ANSICHAR*>& OutExtensions) override;
+	virtual void PostCreateInstance(XrInstance InInstance) override;
+	virtual void PostCreateSession(XrSession InSession) override;
+	virtual void OnDestroySession(XrSession InSession) override;
+	//~ End IOpenXRExtensionPlugin interface
+
+private:
+	XrSpace GetReferenceSpace() const;
+	XrResult GetXrTimeNow(XrTime& OutXrTime) const;
+
+	bool LocateSpace(XrSpace InSpace, XrSpace InRefSpace, XrTime InTime, FTransform& OutTransform);
+
+#if PLATFORM_WINDOWS
+	typedef XrResult(XRAPI_PTR* PFN_xrConvertWin32PerformanceCounterToTimeKHR)(XrInstance Instance, const Windows::LARGE_INTEGER* PerformanceCounter, XrTime* Time);
+	PFN_xrConvertWin32PerformanceCounterToTimeKHR PfnXrConvertWin32PerformanceCounterToTimeKHR = nullptr;
+#endif
+
+private:
+	bool bIsSupported = false;
+
+	XrInstance Instance = XR_NULL_HANDLE;
+
+	XrReferenceSpaceType TrackingSpaceType = XR_REFERENCE_SPACE_TYPE_LOCAL;
+	XrSpace LocalSpace = XR_NULL_HANDLE;
+	XrSpace StageSpace = XR_NULL_HANDLE;
+
+	double WorldToMetersScale;
+};
diff --git a/Engine/Plugins/VirtualProduction/LiveLinkXR/Source/LiveLinkXROpenXRExt/Private/LiveLinkXROpenXRExtModule.cpp b/Engine/Plugins/VirtualProduction/LiveLinkXR/Source/LiveLinkXROpenXRExt/Private/LiveLinkXROpenXRExtModule.cpp
new file mode 100644
--- /dev/null
+++ b/Engine/Plugins/VirtualProduction/LiveLinkXR/Source/LiveLinkXROpenXRExt/Private/LiveLinkXROpenXRExtModule.cpp
@@ -0,0 +1,28 @@
+// Copyright Epic Games, Inc. All Rights Reserved.
+
+#include "LiveLinkXROpenXRExtModule.h"
+#include "LiveLinkXROpenXRExt.h"
+
+
+DEFINE_LOG_CATEGORY(LogLiveLinkXROpenXRExt);
+
+
+IMPLEMENT_MODULE(FLiveLinkXROpenXRExtModule, LiveLinkXROpenXRExt);
+
+
+void FLiveLinkXROpenXRExtModule::StartupModule()
+{
+	OpenXrExt = MakeShared<FLiveLinkXROpenXRExtension>();
+}
+
+
+void FLiveLinkXROpenXRExtModule::ShutdownModule()
+{
+	OpenXrExt.Reset();
+}
+
+
+TSharedPtr<FLiveLinkXROpenXRExtension> FLiveLinkXROpenXRExtModule::GetExtension()
+{
+	return OpenXrExt;
+}
diff --git a/Engine/Plugins/VirtualProduction/LiveLinkXR/Source/LiveLinkXROpenXRExt/Private/LiveLinkXROpenXRExtModule.h b/Engine/Plugins/VirtualProduction/LiveLinkXR/Source/LiveLinkXROpenXRExt/Private/LiveLinkXROpenXRExtModule.h
new file mode 100644
--- /dev/null
+++ b/Engine/Plugins/VirtualProduction/LiveLinkXR/Source/LiveLinkXROpenXRExt/Private/LiveLinkXROpenXRExtModule.h
@@ -0,0 +1,47 @@
+// Copyright Epic Games, Inc. All Rights Reserved.
+
+#pragma once
+
+#include "Logging/LogMacros.h"
+#include "Modules/ModuleManager.h"
+
+
+DECLARE_LOG_CATEGORY_EXTERN(LogLiveLinkXROpenXRExt, Log, All);
+
+
+class FLiveLinkXROpenXRExtension;
+
+
+class LIVELINKXROPENXREXT_API FLiveLinkXROpenXRExtModule : public IModuleInterface
+{
+public:
+	/**
+	* Singleton-like access to this module's interface.  This is just for convenience!
+	* Beware of calling this during the shutdown phase, though.  Your module might have been unloaded already.
+	*
+	* @return Returns singleton instance, loading the module on demand if needed
+	*/
+	static inline FLiveLinkXROpenXRExtModule& Get()
+	{
+		return FModuleManager::LoadModuleChecked<FLiveLinkXROpenXRExtModule>("LiveLinkXROpenXRExt");
+	}
+
+	/**
+	* Checks to see if this module is loaded and ready.  It is only valid to call Get() if IsAvailable() returns true.
+	*
+	* @return True if the module is loaded and ready to use
+	*/
+	static inline bool IsAvailable()
+	{
+		return FModuleManager::Get().IsModuleLoaded("LiveLinkXROpenXRExt");
+	}
+
+	TSharedPtr<FLiveLinkXROpenXRExtension> GetExtension();
+
+private:
+	virtual void StartupModule() override;
+	virtual void ShutdownModule() override;
+
+private:
+	TSharedPtr<FLiveLinkXROpenXRExtension> OpenXrExt;
+};
