Add missing undef LOCTEXT_NAMESPACE

diff --git a/Engine/Plugins/VirtualProduction/LiveLinkXR/Source/LiveLinkXROpenXRExt/Private/LiveLinkXROpenXRExt.cpp b/Engine/Plugins/VirtualProduction/LiveLinkXR/Source/LiveLinkXROpenXRExt/Private/LiveLinkXROpenXRExt.cpp
--- a/Engine/Plugins/VirtualProduction/LiveLinkXR/Source/LiveLinkXROpenXRExt/Private/LiveLinkXROpenXRExt.cpp
+++ b/Engine/Plugins/VirtualProduction/LiveLinkXR/Source/LiveLinkXROpenXRExt/Private/LiveLinkXROpenXRExt.cpp
@@ -224,3 +224,5 @@ bool FLiveLinkXROpenXRExtension::LocateSpace(XrSpace InSpace, XrSpace InRefSpace
 	}
 	return false;
 }
+
+#undef LOCTEXT_NAMESPACE
