LiveLinkXR: Move LiveLinkXR Extension headers to Internal/ to remove LiveLinkXR's dependency on a private include path
#rnx

diff --git a/Engine/Plugins/VirtualProduction/LiveLinkXR/Source/LiveLinkXR/LiveLinkXR.Build.cs b/Engine/Plugins/VirtualProduction/LiveLinkXR/Source/LiveLinkXR/LiveLinkXR.Build.cs
--- a/Engine/Plugins/VirtualProduction/LiveLinkXR/Source/LiveLinkXR/LiveLinkXR.Build.cs
+++ b/Engine/Plugins/VirtualProduction/LiveLinkXR/Source/LiveLinkXR/LiveLinkXR.Build.cs
@@ -8,13 +8,6 @@ public class LiveLinkXR : ModuleRules
 	{
 		PCHUsage = ModuleRules.PCHUsageMode.UseExplicitOrSharedPCHs;
 
-		PrivateIncludePaths.AddRange(
-			new string[]
-			{
-				"LiveLinkXROpenXRExt/Private",
-			}
-		);
-
 		PublicDependencyModuleNames.AddRange(
 			new string[]
 			{
@@ -29,10 +22,10 @@ public class LiveLinkXR : ModuleRules
 				"CoreUObject",
 				"Engine",
 				"HeadMountedDisplay",
+				"LiveLinkXROpenXRExt",
+				"OpenXRHMD",
 				"Slate",
 				"SlateCore",
-				"OpenXRHMD",
-				"LiveLinkXROpenXRExt",
 			}
 		);
 	}
diff --git a/Engine/Plugins/VirtualProduction/LiveLinkXR/Source/LiveLinkXROpenXRExt/Internal/LiveLinkXROpenXRExt.h b/Engine/Plugins/VirtualProduction/LiveLinkXR/Source/LiveLinkXROpenXRExt/Internal/LiveLinkXROpenXRExt.h
new file mode 100644
--- /dev/null
+++ b/Engine/Plugins/VirtualProduction/LiveLinkXR/Source/LiveLinkXROpenXRExt/Internal/LiveLinkXROpenXRExt.h
@@ -0,0 +1,61 @@
+// Copyright Epic Games, Inc. All Rights Reserved.
+
+#pragma once
+
+#include "IOpenXRExtensionPlugin.h"
+#include "OpenXRCore.h"
+
+
+#if PLATFORM_WINDOWS
+namespace Windows
+{
+	union LARGE_INTEGER;
+}
+#endif
+
+
+class LIVELINKXROPENXREXT_API FLiveLinkXROpenXRExtension
+	: public IOpenXRExtensionPlugin
+	, public TSharedFromThis<FLiveLinkXROpenXRExtension>
+{
+public:
+	FLiveLinkXROpenXRExtension();
+	virtual ~FLiveLinkXROpenXRExtension();
+
+	bool IsSupported() const { return bIsSupported; }
+	void GetSubjectPoses(TMap<FName, FTransform>& OutPoseMap);
+
+	//~ Begin IOpenXRExtensionPlugin interface
+	virtual FString GetDisplayName() override
+	{
+		return FString(TEXT("LiveLinkXR"));
+	}
+
+	virtual bool GetOptionalExtensions(TArray<const ANSICHAR*>& OutExtensions) override;
+	virtual void PostCreateInstance(XrInstance InInstance) override;
+	virtual void PostCreateSession(XrSession InSession) override;
+	virtual void OnDestroySession(XrSession InSession) override;
+	//~ End IOpenXRExtensionPlugin interface
+
+private:
+	XrSpace GetReferenceSpace() const;
+	XrResult GetXrTimeNow(XrTime& OutXrTime) const;
+
+	bool LocateSpace(XrSpace InSpace, XrSpace InRefSpace, XrTime InTime, FTransform& OutTransform);
+
+#if PLATFORM_WINDOWS
+	typedef XrResult(XRAPI_PTR* PFN_xrConvertWin32PerformanceCounterToTimeKHR)(XrInstance Instance, const Windows::LARGE_INTEGER* PerformanceCounter, XrTime* Time);
+	PFN_xrConvertWin32PerformanceCounterToTimeKHR PfnXrConvertWin32PerformanceCounterToTimeKHR = nullptr;
+#endif
+
+private:
+	bool bIsSupported = false;
+
+	XrInstance Instance = XR_NULL_HANDLE;
+
+	XrReferenceSpaceType TrackingSpaceType = XR_REFERENCE_SPACE_TYPE_LOCAL;
+	XrSpace LocalSpace = XR_NULL_HANDLE;
+	XrSpace StageSpace = XR_NULL_HANDLE;
+
+	double WorldToMetersScale;
+};
diff --git a/Engine/Plugins/VirtualProduction/LiveLinkXR/Source/LiveLinkXROpenXRExt/Internal/LiveLinkXROpenXRExtModule.h b/Engine/Plugins/VirtualProduction/LiveLinkXR/Source/LiveLinkXROpenXRExt/Internal/LiveLinkXROpenXRExtModule.h
new file mode 100644
--- /dev/null
+++ b/Engine/Plugins/VirtualProduction/LiveLinkXR/Source/LiveLinkXROpenXRExt/Internal/LiveLinkXROpenXRExtModule.h
@@ -0,0 +1,47 @@
+// Copyright Epic Games, Inc. All Rights Reserved.
+
+#pragma once
+
+#include "Logging/LogMacros.h"
+#include "Modules/ModuleManager.h"
+
+
+DECLARE_LOG_CATEGORY_EXTERN(LogLiveLinkXROpenXRExt, Log, All);
+
+
+class FLiveLinkXROpenXRExtension;
+
+
+class LIVELINKXROPENXREXT_API FLiveLinkXROpenXRExtModule : public IModuleInterface
+{
+public:
+	/**
+	* Singleton-like access to this module's interface.  This is just for convenience!
+	* Beware of calling this during the shutdown phase, though.  Your module might have been unloaded already.
+	*
+	* @return Returns singleton instance, loading the module on demand if needed
+	*/
+	static inline FLiveLinkXROpenXRExtModule& Get()
+	{
+		return FModuleManager::LoadModuleChecked<FLiveLinkXROpenXRExtModule>("LiveLinkXROpenXRExt");
+	}
+
+	/**
+	* Checks to see if this module is loaded and ready.  It is only valid to call Get() if IsAvailable() returns true.
+	*
+	* @return True if the module is loaded and ready to use
+	*/
+	static inline bool IsAvailable()
+	{
+		return FModuleManager::Get().IsModuleLoaded("LiveLinkXROpenXRExt");
+	}
+
+	TSharedPtr<FLiveLinkXROpenXRExtension> GetExtension();
+
+private:
+	virtual void StartupModule() override;
+	virtual void ShutdownModule() override;
+
+private:
+	TSharedPtr<FLiveLinkXROpenXRExtension> OpenXrExt;
+};
diff --git a/Engine/Plugins/VirtualProduction/LiveLinkXR/Source/LiveLinkXROpenXRExt/Private/LiveLinkXROpenXRExt.h b/Engine/Plugins/VirtualProduction/LiveLinkXR/Source/LiveLinkXROpenXRExt/Private/LiveLinkXROpenXRExt.h
deleted file mode 100644
diff --git a/Engine/Plugins/VirtualProduction/LiveLinkXR/Source/LiveLinkXROpenXRExt/Private/LiveLinkXROpenXRExtModule.h b/Engine/Plugins/VirtualProduction/LiveLinkXR/Source/LiveLinkXROpenXRExt/Private/LiveLinkXROpenXRExtModule.h
deleted file mode 100644
