LiveLinkXR: Remove Linux from supported platforms.

In theory it would not be difficult to add support for XR_KHR_convert_timespec_time, but I'm not set up to test it.

#rnx

diff --git a/Engine/Plugins/VirtualProduction/LiveLinkXR/LiveLinkXR.uplugin b/Engine/Plugins/VirtualProduction/LiveLinkXR/LiveLinkXR.uplugin
--- a/Engine/Plugins/VirtualProduction/LiveLinkXR/LiveLinkXR.uplugin
+++ b/Engine/Plugins/VirtualProduction/LiveLinkXR/LiveLinkXR.uplugin
@@ -16,8 +16,7 @@
 	"Installed": false,
 	"EnabledByDefault": false,
 	"SupportedTargetPlatforms": [
-		"Win64",
-		"Linux"
+		"Win64"
 	],
 	"Plugins": [
 		{
@@ -36,8 +35,7 @@
 			"Type": "Runtime",
 			"LoadingPhase": "PostConfigInit",
 			"PlatformAllowList": [
-				"Win64",
-				"Linux"
+				"Win64"
 			]
 		}
 	]
