LiveLinkXR: Minor refactor to remove depedency on private header in OpenXRHMD module
#rnx

diff --git a/Engine/Plugins/VirtualProduction/LiveLinkXR/Source/LiveLinkXROpenXRExt/LiveLinkXROpenXRExt.Build.cs b/Engine/Plugins/VirtualProduction/LiveLinkXR/Source/LiveLinkXROpenXRExt/LiveLinkXROpenXRExt.Build.cs
--- a/Engine/Plugins/VirtualProduction/LiveLinkXR/Source/LiveLinkXROpenXRExt/LiveLinkXROpenXRExt.Build.cs
+++ b/Engine/Plugins/VirtualProduction/LiveLinkXR/Source/LiveLinkXROpenXRExt/LiveLinkXROpenXRExt.Build.cs
@@ -9,13 +9,6 @@ public class LiveLinkXROpenXRExt : ModuleRules
 	{
 		PCHUsage = ModuleRules.PCHUsageMode.UseExplicitOrSharedPCHs;
 
-		var EngineDir = Path.GetFullPath(Target.RelativeEnginePath);
-		PrivateIncludePaths.AddRange(
-            new string[] {
-                EngineDir + "/Plugins/Runtime/OpenXR/Source/OpenXRHMD/Private",
-            }
-		);
-
 		PublicDependencyModuleNames.AddRange(
 			new string[]
 			{
diff --git a/Engine/Plugins/VirtualProduction/LiveLinkXR/Source/LiveLinkXROpenXRExt/Private/LiveLinkXROpenXRExt.cpp b/Engine/Plugins/VirtualProduction/LiveLinkXR/Source/LiveLinkXROpenXRExt/Private/LiveLinkXROpenXRExt.cpp
--- a/Engine/Plugins/VirtualProduction/LiveLinkXR/Source/LiveLinkXROpenXRExt/Private/LiveLinkXROpenXRExt.cpp
+++ b/Engine/Plugins/VirtualProduction/LiveLinkXR/Source/LiveLinkXROpenXRExt/Private/LiveLinkXROpenXRExt.cpp
@@ -4,10 +4,10 @@
 #include "LiveLinkXROpenXRExtModule.h"
 
 #include "Engine/Engine.h"
+#include "IOpenXRHMD.h"
 #include "IOpenXRHMDModule.h"
 #include "IXRTrackingSystem.h"
 #include "OpenXRCore.h"
-#include "OpenXRHMD.h"
 
 #if PLATFORM_WINDOWS
 	#include "Microsoft/MinimalWindowsApi.h"
@@ -59,16 +59,22 @@ void FLiveLinkXROpenXRExtension::GetSubjectPoses(TMap<FName, FTransform>& OutPos
 	}
 
 	TArray<int32> Devices;
-	FOpenXRHMD* OpenXRHMD = (FOpenXRHMD*)GEngine->XRSystem.Get();
-	if (OpenXRHMD->EnumerateTrackedDevices(Devices))
+
+	if (GEngine && GEngine->XRSystem.IsValid())
 	{
-		for (int32 DeviceId : Devices)
+		IXRTrackingSystem* XRTrackingSystem = GEngine->XRSystem.Get();
+		IOpenXRHMD* OpenXRHMD = XRTrackingSystem->GetIOpenXRHMD();
+
+		if (XRTrackingSystem->EnumerateTrackedDevices(Devices))
 		{
-			FName Path;
-			FTransform Location;
-			if (LocateSpace(OpenXRHMD->GetTrackedDeviceSpace(DeviceId), RefSpace, PoseTime, Location))
+			for (int32 DeviceId : Devices)
 			{
-				OutPoseMap.Emplace(FOpenXRPath(OpenXRHMD->GetTrackedDevicePath(DeviceId)), Location);
+				FName Path;
+				FTransform Location;
+				if (LocateSpace(OpenXRHMD->GetTrackedDeviceSpace(DeviceId), RefSpace, PoseTime, Location))
+				{
+					OutPoseMap.Emplace(FOpenXRPath(OpenXRHMD->GetTrackedDevicePath(DeviceId)), Location);
+				}
 			}
 		}
 	}
