OpenXR extension names: use openxr.h define rather than hard coding the names.
-We should use the header defines rather than defining the identical string ourselves.

#rnx

diff --git a/Engine/Plugins/Runtime/HPMotionController/Source/HPMotionController/Private/HPMotionController.cpp b/Engine/Plugins/Runtime/HPMotionController/Source/HPMotionController/Private/HPMotionController.cpp
--- a/Engine/Plugins/Runtime/HPMotionController/Source/HPMotionController/Private/HPMotionController.cpp
+++ b/Engine/Plugins/Runtime/HPMotionController/Source/HPMotionController/Private/HPMotionController.cpp
@@ -113,7 +113,7 @@ void FHPMotionControllerModule::ShutdownModule()
 bool FHPMotionControllerModule::GetRequiredExtensions(TArray<const ANSICHAR*>& OutExtensions)
 {
 	// XR_EXT_HP_MIXED_REALITY_CONTROLLER_EXTENSION_NAME is defined in openxr.h starting with 1.0.10
-	OutExtensions.Add("XR_EXT_hp_mixed_reality_controller");
+	OutExtensions.Add(XR_EXT_HP_MIXED_REALITY_CONTROLLER_EXTENSION_NAME);
 	return true;
 }
 
diff --git a/Engine/Plugins/Runtime/OpenXREyeTracker/Source/OpenXREyeTracker/Private/OpenXREyeTracker.cpp b/Engine/Plugins/Runtime/OpenXREyeTracker/Source/OpenXREyeTracker/Private/OpenXREyeTracker.cpp
--- a/Engine/Plugins/Runtime/OpenXREyeTracker/Source/OpenXREyeTracker/Private/OpenXREyeTracker.cpp
+++ b/Engine/Plugins/Runtime/OpenXREyeTracker/Source/OpenXREyeTracker/Private/OpenXREyeTracker.cpp
@@ -44,7 +44,7 @@ bool FOpenXREyeTracker::GetRequiredExtensions(TArray<const ANSICHAR*>& OutExtens
 		}
 	}
 
-	OutExtensions.Add("XR_EXT_eye_gaze_interaction");
+	OutExtensions.Add(XR_EXT_EYE_GAZE_INTERACTION_EXTENSION_NAME);
 	return true;
 }
 
diff --git a/Engine/Plugins/Runtime/OpenXRMsftHandInteraction/Source/OpenXRMsftHandInteraction/Private/OpenXRMsftHandInteraction.cpp b/Engine/Plugins/Runtime/OpenXRMsftHandInteraction/Source/OpenXRMsftHandInteraction/Private/OpenXRMsftHandInteraction.cpp
--- a/Engine/Plugins/Runtime/OpenXRMsftHandInteraction/Source/OpenXRMsftHandInteraction/Private/OpenXRMsftHandInteraction.cpp
+++ b/Engine/Plugins/Runtime/OpenXRMsftHandInteraction/Source/OpenXRMsftHandInteraction/Private/OpenXRMsftHandInteraction.cpp
@@ -33,7 +33,7 @@ void FOpenXRMsftHandInteraction::StartupModule()
 
 bool FOpenXRMsftHandInteraction::GetRequiredExtensions(TArray<const ANSICHAR*>& OutExtensions)
 {
-	OutExtensions.Add("XR_MSFT_hand_interaction");
+	OutExtensions.Add(XR_MSFT_HAND_INTERACTION_EXTENSION_NAME);
 	return true;
 }
 
diff --git a/Engine/Plugins/VirtualProduction/LiveLinkXR/Source/LiveLinkXROpenXRExt/Private/LiveLinkXROpenXRExt.cpp b/Engine/Plugins/VirtualProduction/LiveLinkXR/Source/LiveLinkXROpenXRExt/Private/LiveLinkXROpenXRExt.cpp
--- a/Engine/Plugins/VirtualProduction/LiveLinkXR/Source/LiveLinkXROpenXRExt/Private/LiveLinkXROpenXRExt.cpp
+++ b/Engine/Plugins/VirtualProduction/LiveLinkXR/Source/LiveLinkXROpenXRExt/Private/LiveLinkXROpenXRExt.cpp
@@ -8,6 +8,7 @@
 #include "IOpenXRHMDModule.h"
 #include "IXRTrackingSystem.h"
 #include "OpenXRCore.h"
+#include "OpenXRPlatformRHI.h" // Needed to get openxr_platform.h with the correct #defines set
 
 #if PLATFORM_WINDOWS
 	#include "Microsoft/MinimalWindowsApi.h"
@@ -84,7 +85,7 @@ void FLiveLinkXROpenXRExtension::GetSubjectPoses(TMap<FName, FTransform>& OutPos
 bool FLiveLinkXROpenXRExtension::GetOptionalExtensions(TArray<const ANSICHAR*>& OutExtensions)
 {
 #if PLATFORM_WINDOWS
-	OutExtensions.Add("XR_KHR_win32_convert_performance_counter_time");
+	OutExtensions.Add(XR_KHR_WIN32_CONVERT_PERFORMANCE_COUNTER_TIME_EXTENSION_NAME);
 #endif
 
 	return true;
